[ChaosVD] Fix for cooked CVD Standalone crashing at start up

This CL also fixes an issue where the standalone app fails to detect active client/server instances (the UdpMessaging plugin was not being loaded), and lastly, removed some content being staged/cooked that is not needed.

diff --git a/Engine/Plugins/Messaging/UdpMessaging/UdpMessaging.uplugin b/Engine/Plugins/Messaging/UdpMessaging/UdpMessaging.uplugin
--- a/Engine/Plugins/Messaging/UdpMessaging/UdpMessaging.uplugin
+++ b/Engine/Plugins/Messaging/UdpMessaging/UdpMessaging.uplugin
@@ -25,7 +25,8 @@
 		"LiveLinkHub",
 		"TimeSynchronizationWriter",
 		"LiveLinkProviderSample",
-		"SpatialMetricsProfiler"
+		"SpatialMetricsProfiler",
+		"ChaosVisualDebugger"
 	],
 	"Modules" :
 	[
@@ -47,7 +48,8 @@
 				"LiveLinkHub",
 				"TimeSynchronizationWriter",
 				"LiveLinkProviderSample",
-				"SpatialMetricsProfiler"
+				"SpatialMetricsProfiler",
+				"ChaosVisualDebugger"
 			]
 		}
 	]
diff --git a/Engine/Programs/ChaosVisualDebugger/ChaosVisualDebugger.uproject b/Engine/Programs/ChaosVisualDebugger/ChaosVisualDebugger.uproject
--- a/Engine/Programs/ChaosVisualDebugger/ChaosVisualDebugger.uproject
+++ b/Engine/Programs/ChaosVisualDebugger/ChaosVisualDebugger.uproject
@@ -14,6 +14,10 @@
 		{
 			"Name": "PropertyAccessEditor",
 			"Enabled": true
-		}	
+		},
+		{
+			"Name": "UdpMessaging",
+			"Enabled": true
+		}
 	]
 }
\ No newline at end of file
diff --git a/Engine/Programs/ChaosVisualDebugger/Config/DefaultGame.ini b/Engine/Programs/ChaosVisualDebugger/Config/DefaultGame.ini
--- a/Engine/Programs/ChaosVisualDebugger/Config/DefaultGame.ini
+++ b/Engine/Programs/ChaosVisualDebugger/Config/DefaultGame.ini
@@ -8,7 +8,7 @@ bUseZenStore=false
 bSupportCookedEditor=true
 CookedEditorTargetName=ChaosVisualDebugger
 bStageUAT=false
-bStageShaderDirs=true
+bStageShaderDirs=false
 bStagePlatformBuildDirs=false
 bStageExtrasDirs=false
 bStagePlatformDirs=false
@@ -24,5 +24,3 @@ MapMode=None
 ; this will remove all of the compiled UAT script dlls, except for some ones that had their own subdirectories when built, that we need to keep
 +FinalStripFiles=(Filter="AutomationTool/AutomationScripts/", Exceptions="/(BuildGraph|BuildPatchTool|ContentWorker|Scripts|Android)/")
 
-; Always cook the control rig standard function library
-+EngineSpecificAssetsToCook="/ControlRig/StandardFunctionLibrary/StandardFunctionLibrary.uasset"
diff --git a/Engine/Source/Programs/ChaosVisualDebugger/ChaosVisualDebugger.Target.cs b/Engine/Source/Programs/ChaosVisualDebugger/ChaosVisualDebugger.Target.cs
--- a/Engine/Source/Programs/ChaosVisualDebugger/ChaosVisualDebugger.Target.cs
+++ b/Engine/Source/Programs/ChaosVisualDebugger/ChaosVisualDebugger.Target.cs
@@ -56,7 +56,7 @@ public class ChaosVisualDebuggerTarget : TargetRules
 		GlobalDefinitions.Add("ASSETREGISTRY_FORCE_PREMADE_REGISTRY_IN_EDITOR=1");
 		GlobalDefinitions.Add("UE_IS_COOKED_EDITOR=1");
 		GlobalDefinitions.Add("UE_FORCE_USE_PAKS=1");
-		GlobalDefinitions.Add("UE_FORCE_USE_IOSTORE=1");
+		GlobalDefinitions.Add("UE_FORCE_USE_IOSTORE=0");
 
 		// this will allow shader compiling to work based on whether or not the shaders directory is present
 		// to determine if we should allow shader compilation
diff --git a/Engine/Source/Programs/ChaosVisualDebugger/Private/ChaosVisualDebuggerMain.cpp b/Engine/Source/Programs/ChaosVisualDebugger/Private/ChaosVisualDebuggerMain.cpp
--- a/Engine/Source/Programs/ChaosVisualDebugger/Private/ChaosVisualDebuggerMain.cpp
+++ b/Engine/Source/Programs/ChaosVisualDebugger/Private/ChaosVisualDebuggerMain.cpp
@@ -48,7 +48,7 @@ int32 RunChaosVisualDebugger(const TCHAR* CommandLine)
 	const bool bIsRunningCommand = FParse::Value(CommandLine, TEXT("-RUN="), Command);
 
 	const FString CommandLineString = CommandLine;
-	const FString FinalCommandLine = bIsRunningCommand ? CommandLine : CommandLineString + TEXT(" EDITOR");
+	const FString FinalCommandLine = bIsRunningCommand ? CommandLine : CommandLineString + TEXT(" EDITOR -messaging");
 
 	// start up the main loop
 	const int32 Result = GEngineLoop.PreInit(*FinalCommandLine);
